Basic prefab plugin, no runtime data yet

diff --git a/Engine/Plugins/Experimental/Prefabs/Prefabs.uplugin b/Engine/Plugins/Experimental/Prefabs/Prefabs.uplugin
new file mode 100644
--- /dev/null
+++ b/Engine/Plugins/Experimental/Prefabs/Prefabs.uplugin
@@ -0,0 +1,34 @@
+{
+	"FileVersion" : 3,
+	"Version" : 1,
+	"VersionName" : "1.0",
+	"FriendlyName" : "Prefabs",
+	"Description" : "Experimental placement and archetype management system for the Unreal Engine",
+	"Category" : "Gameplay",
+	"CreatedBy" : "Epic Games, Inc.",
+	"CreatedByURL" : "http://epicgames.com",
+	"DocsURL" : "",
+	"MarketplaceURL" : "",
+	"SupportURL" : "",
+	"EnabledByDefault" : false,
+	"CanContainContent" : true,
+	"IsExperimentalVersion" : true,
+	"Installed" : false,
+	"Modules": [
+		{
+			"Name": "PrefabsRuntime",
+			"Type": "Runtime",
+			"LoadingPhase": "Default"
+		},
+		{
+			"Name": "PrefabsUncooked",
+			"Type": "UncookedOnly",
+			"LoadingPhase": "Default"
+		},
+		{
+			"Name": "PrefabsEditor",
+			"Type": "Editor",
+			"LoadingPhase": "Default"
+		}
+	]
+}
\ No newline at end of file
diff --git a/Engine/Plugins/Experimental/Prefabs/Source/PrefabsEditor/PrefabsEditor.Build.cs b/Engine/Plugins/Experimental/Prefabs/Source/PrefabsEditor/PrefabsEditor.Build.cs
new file mode 100644
--- /dev/null
+++ b/Engine/Plugins/Experimental/Prefabs/Source/PrefabsEditor/PrefabsEditor.Build.cs
@@ -0,0 +1,27 @@
+// Copyright Epic Games, Inc. All Rights Reserved.
+
+namespace UnrealBuildTool.Rules
+{
+	public class PrefabsEditor : ModuleRules
+	{
+		public PrefabsEditor(ReadOnlyTargetRules Target) : base(Target)
+        {
+			PrivateDependencyModuleNames.AddRange(
+				new string[]
+				{
+					"AssetTools",
+					"Core",
+					"CoreUObject",
+					"EditorFramework",
+					"Engine",
+					"PrefabsUncooked",
+					"Slate",
+					"SlateCore",
+					"TargetPlatform",
+					"ToolMenus",
+					"UnrealEd",
+				}
+			); 
+		}
+	}
+}
diff --git a/Engine/Plugins/Experimental/Prefabs/Source/PrefabsEditor/Private/PrefabAssetTypeActions.cpp b/Engine/Plugins/Experimental/Prefabs/Source/PrefabsEditor/Private/PrefabAssetTypeActions.cpp
new file mode 100644
--- /dev/null
+++ b/Engine/Plugins/Experimental/Prefabs/Source/PrefabsEditor/Private/PrefabAssetTypeActions.cpp
@@ -0,0 +1,60 @@
+// Copyright Epic Games, Inc. All Rights Reserved.
+
+#include "PrefabAssetTypeActions.h"
+#include "PrefabUncooked.h"
+#include "ToolMenuSection.h"
+
+#define LOCTEXT_NAMESPACE "PrefabAssetTypeActions"
+
+//////////////////////////////////////////////////////////////////////////
+// FPrefabAssetTypeActions
+
+FPrefabAssetTypeActions::FPrefabAssetTypeActions(EAssetTypeCategories::Type InAssetCategory)
+	: MyAssetCategory(InAssetCategory)
+{
+}
+
+FText FPrefabAssetTypeActions::GetName() const
+{
+	return LOCTEXT("FPrefabAssetTypeActionsName", "Prefab");
+}
+
+FColor FPrefabAssetTypeActions::GetTypeColor() const
+{
+	return FColorList::Orange;
+}
+
+UClass* FPrefabAssetTypeActions::GetSupportedClass() const
+{
+	return UPrefabUncooked::StaticClass();
+}
+
+void FPrefabAssetTypeActions::OpenAssetEditor(const TArray<UObject*>& InObjects, TSharedPtr<class IToolkitHost> EditWithinLevelEditor)
+{
+	FSimpleAssetEditor::CreateEditor(EToolkitMode::Standalone, EditWithinLevelEditor, InObjects);
+}
+
+uint32 FPrefabAssetTypeActions::GetCategories()
+{
+	return MyAssetCategory;
+}
+
+void FPrefabAssetTypeActions::GetActions(const TArray<UObject*>& InObjects, FToolMenuSection& Section)
+{
+	auto Prefabs = GetTypedWeakObjectPtrs<UPrefabUncooked>(InObjects);
+
+	Section.AddMenuEntry(
+		"CreatePrefab",
+		LOCTEXT("CreatePrefab", "Create Prefab"),
+		LOCTEXT("CreatePrefabTooltip", "Create an empty Prefab."),
+		FSlateIcon(FAppStyle::GetAppStyleSetName(), "ClassIcon.Prefab"),
+		FUIAction(
+			FExecuteAction::CreateLambda([] {}),
+			FCanExecuteAction()
+		)
+	);
+}
+
+//////////////////////////////////////////////////////////////////////////
+
+#undef LOCTEXT_NAMESPACE
diff --git a/Engine/Plugins/Experimental/Prefabs/Source/PrefabsEditor/Private/PrefabAssetTypeActions.h b/Engine/Plugins/Experimental/Prefabs/Source/PrefabsEditor/Private/PrefabAssetTypeActions.h
new file mode 100644
--- /dev/null
+++ b/Engine/Plugins/Experimental/Prefabs/Source/PrefabsEditor/Private/PrefabAssetTypeActions.h
@@ -0,0 +1,28 @@
+// Copyright Epic Games, Inc. All Rights Reserved.
+
+#pragma once
+
+#include "CoreMinimal.h"
+#include "Toolkits/IToolkitHost.h"
+#include "AssetTypeActions_Base.h"
+
+struct FAssetData;
+
+class FPrefabAssetTypeActions : public FAssetTypeActions_Base
+{
+public:
+	FPrefabAssetTypeActions(EAssetTypeCategories::Type InAssetCategory);
+
+	// IAssetTypeActions interface
+	virtual FText GetName() const override;
+	virtual FColor GetTypeColor() const override;
+	virtual UClass* GetSupportedClass() const override;
+	virtual void OpenAssetEditor(const TArray<UObject*>& InObjects, TSharedPtr<class IToolkitHost> EditWithinLevelEditor = TSharedPtr<IToolkitHost>()) override;
+	virtual uint32 GetCategories() override;
+	virtual bool HasActions(const TArray<UObject*>& InObjects) const override { return true; }
+	virtual void GetActions(const TArray<UObject*>& InObjects, struct FToolMenuSection& Section) override;
+	// End of IAssetTypeActions interface
+
+private:
+	EAssetTypeCategories::Type MyAssetCategory;
+};
diff --git a/Engine/Plugins/Experimental/Prefabs/Source/PrefabsEditor/Private/PrefabFactory.cpp b/Engine/Plugins/Experimental/Prefabs/Source/PrefabsEditor/Private/PrefabFactory.cpp
new file mode 100644
--- /dev/null
+++ b/Engine/Plugins/Experimental/Prefabs/Source/PrefabsEditor/Private/PrefabFactory.cpp
@@ -0,0 +1,23 @@
+// Copyright Epic Games, Inc. All Rights Reserved.
+
+#include "PrefabFactory.h"
+#include "PrefabUncooked.h"
+
+UPrefabFactory::UPrefabFactory(const FObjectInitializer& ObjectInitializer)
+	: Super(ObjectInitializer)
+{
+	bCreateNew = true;
+	bEditAfterNew = true;
+	SupportedClass = UPrefabUncooked::StaticClass();
+}
+
+UObject* UPrefabFactory::FactoryCreateNew( UClass* Class, UObject* InParent, FName Name, EObjectFlags Flags, UObject* Context, FFeedbackContext* Warn)
+{
+	UPrefabUncooked* NewPrefab = NewObject<UPrefabUncooked>(InParent, Class, Name, Flags | RF_Transactional);
+	return NewPrefab;
+}
+
+FString UPrefabFactory::GetDefaultNewAssetName() const
+{
+	return TEXT("NewPrefab");
+}
diff --git a/Engine/Plugins/Experimental/Prefabs/Source/PrefabsEditor/Private/PrefabFactory.h b/Engine/Plugins/Experimental/Prefabs/Source/PrefabsEditor/Private/PrefabFactory.h
new file mode 100644
--- /dev/null
+++ b/Engine/Plugins/Experimental/Prefabs/Source/PrefabsEditor/Private/PrefabFactory.h
@@ -0,0 +1,22 @@
+// Copyright Epic Games, Inc. All Rights Reserved.
+
+#pragma once
+
+#include "CoreMinimal.h"
+#include "UObject/ObjectMacros.h"
+#include "Factories/Factory.h"
+#include "PrefabFactory.generated.h"
+
+// Asset factory boilerplate, allows the editor to create prefab assets:
+UCLASS()
+class UPrefabFactory : public UFactory
+{
+	GENERATED_BODY()
+
+	UPrefabFactory(const FObjectInitializer& ObjectInitializer);
+
+	// UFactory interface
+	virtual UObject* FactoryCreateNew(UClass* Class, UObject* InParent, FName Name, EObjectFlags Flags, UObject* Context, FFeedbackContext* Warn) override;
+	virtual FString GetDefaultNewAssetName() const override;
+	// End of UFactory interface
+};
diff --git a/Engine/Plugins/Experimental/Prefabs/Source/PrefabsEditor/Private/PrefabsEditorModule.cpp b/Engine/Plugins/Experimental/Prefabs/Source/PrefabsEditor/Private/PrefabsEditorModule.cpp
new file mode 100644
--- /dev/null
+++ b/Engine/Plugins/Experimental/Prefabs/Source/PrefabsEditor/Private/PrefabsEditorModule.cpp
@@ -0,0 +1,74 @@
+// Copyright Epic Games, Inc. All Rights Reserved.
+
+#include "IPrefabsEditorModule.h"
+#include "AssetToolsModule.h"
+#include "Containers/UnrealString.h"
+#include "Logging/LogMacros.h"
+#include "PrefabAssetTypeActions.h"
+#include "Templates/Function.h"
+
+DECLARE_LOG_CATEGORY_EXTERN(PrefabsEditor, Log, All);
+
+/**
+ * The Plugin Warden is a simple module used to verify a user has purchased a plug-in.  This
+ * module won't prevent a determined user from avoiding paying for a plug-in, it is merely to
+ * prevent accidental violation of a per-seat license on a plug-in, and to direct those users
+ * to the marketplace page where they may purchase the plug-in.
+ */
+class FPrefabsEditorModule : public IPrefabsEditorModule
+{
+public:
+	virtual void StartupModule() override;
+	virtual void ShutdownModule() override;
+
+private:
+	void OnPostEngineInit();
+	void RegisterAssetTypeAction(IAssetTools& AssetTools, TSharedRef<IAssetTypeActions> Action);
+
+	TArray< TSharedPtr<IAssetTypeActions> > CreatedAssetTypeActions;
+};
+
+IMPLEMENT_MODULE( FPrefabsEditorModule, PrefabsEditor);
+
+#define LOCTEXT_NAMESPACE "PrefabsEditor"
+
+TSet<FString> AuthorizedPlugins;
+
+void FPrefabsEditorModule::StartupModule()
+{
+	FCoreDelegates::OnPostEngineInit.AddRaw(this, &FPrefabsEditorModule::OnPostEngineInit);
+}
+
+void FPrefabsEditorModule::ShutdownModule()
+{
+	if (FModuleManager::Get().IsModuleLoaded("AssetTools"))
+	{
+		IAssetTools& AssetTools = FModuleManager::GetModuleChecked<FAssetToolsModule>("AssetTools").Get();
+		for (int32 Index = 0; Index < CreatedAssetTypeActions.Num(); ++Index)
+		{
+			AssetTools.UnregisterAssetTypeActions(CreatedAssetTypeActions[Index].ToSharedRef());
+		}
+	}
+	CreatedAssetTypeActions.Empty();
+
+	FCoreDelegates::OnPostEngineInit.RemoveAll(this);
+}
+
+void FPrefabsEditorModule::OnPostEngineInit()
+{
+	// Register asset types
+	IAssetTools& AssetTools = FModuleManager::LoadModuleChecked<FAssetToolsModule>("AssetTools").Get();
+	
+	EAssetTypeCategories::Type Category = AssetTools.RegisterAdvancedAssetCategory(
+		FName(TEXT("Prefabs")), LOCTEXT("PrefabsAssetCategory", "Prefabs"));
+
+	RegisterAssetTypeAction(AssetTools, MakeShareable(new FPrefabAssetTypeActions(Category)));
+}
+
+void FPrefabsEditorModule::RegisterAssetTypeAction(IAssetTools& AssetTools, TSharedRef<IAssetTypeActions> Action)
+{
+	AssetTools.RegisterAssetTypeActions(Action);
+	CreatedAssetTypeActions.Add(Action);
+}
+
+#undef LOCTEXT_NAMESPACE
diff --git a/Engine/Plugins/Experimental/Prefabs/Source/PrefabsEditor/Public/IPrefabsEditorModule.h b/Engine/Plugins/Experimental/Prefabs/Source/PrefabsEditor/Public/IPrefabsEditorModule.h
new file mode 100644
--- /dev/null
+++ b/Engine/Plugins/Experimental/Prefabs/Source/PrefabsEditor/Public/IPrefabsEditorModule.h
@@ -0,0 +1,37 @@
+// Copyright Epic Games, Inc. All Rights Reserved.
+#pragma once
+
+#include "CoreMinimal.h"
+#include "Modules/ModuleInterface.h"
+#include "Modules/ModuleManager.h"
+
+/**
+ * Prefabs are an invisible, linkable instantiation routine (function object or
+ * type constructor). They are used to define and later instantiate run time objects.
+ * This 'editor' module provides user interface for an implementation of the former.
+ */
+class IPrefabsEditorModule : public IModuleInterface
+{
+public:
+
+	/**
+	 * Singleton-like access to this module's interface.  This is just for convenience!
+	 * Beware of calling this during the shutdown phase, though.  Your module might have been unloaded already.
+	 *
+	 * @return Returns singleton instance, loading the module on demand if needed
+	 */
+	static inline IPrefabsEditorModule& Get()
+	{
+		return FModuleManager::LoadModuleChecked< IPrefabsEditorModule >( "PrefabsEditor" );
+	}
+
+	/**
+	 * Checks to see if this module is loaded and ready.  It is only valid to call Get() if IsAvailable() returns true.
+	 *
+	 * @return True if the module is loaded and ready to use
+	 */
+	static inline bool IsAvailable()
+	{
+		return FModuleManager::Get().IsModuleLoaded( "PrefabsEditor" );
+	}
+};
diff --git a/Engine/Plugins/Experimental/Prefabs/Source/PrefabsRuntime/PrefabsRuntime.Build.cs b/Engine/Plugins/Experimental/Prefabs/Source/PrefabsRuntime/PrefabsRuntime.Build.cs
new file mode 100644
--- /dev/null
+++ b/Engine/Plugins/Experimental/Prefabs/Source/PrefabsRuntime/PrefabsRuntime.Build.cs
@@ -0,0 +1,18 @@
+// Copyright Epic Games, Inc. All Rights Reserved.
+
+namespace UnrealBuildTool.Rules
+{
+	public class PrefabsRuntime : ModuleRules
+	{
+		public PrefabsRuntime(ReadOnlyTargetRules Target) : base(Target)
+        {
+			PrivateDependencyModuleNames.AddRange(
+				new string[]
+				{
+					"Core",
+					"CoreUObject",
+                }
+			); 
+		}
+	}
+}
diff --git a/Engine/Plugins/Experimental/Prefabs/Source/PrefabsRuntime/Private/PrefabsRuntimeModule.cpp b/Engine/Plugins/Experimental/Prefabs/Source/PrefabsRuntime/Private/PrefabsRuntimeModule.cpp
new file mode 100644
--- /dev/null
+++ b/Engine/Plugins/Experimental/Prefabs/Source/PrefabsRuntime/Private/PrefabsRuntimeModule.cpp
@@ -0,0 +1,29 @@
+// Copyright Epic Games, Inc. All Rights Reserved.
+
+#include "IPrefabsRuntimeModule.h"
+#include "Containers/UnrealString.h"
+#include "Logging/LogMacros.h"
+#include "Templates/Function.h"
+
+DECLARE_LOG_CATEGORY_EXTERN(PrefabsRuntime, Log, All);
+
+class FPrefabsRuntimeModule : public IPrefabsRuntimeModule
+{
+protected:
+	virtual void StartupModule() override;
+	virtual void ShutdownModule() override;
+};
+
+IMPLEMENT_MODULE( FPrefabsRuntimeModule, PrefabsRuntime );
+
+#define LOCTEXT_NAMESPACE "PrefabsRuntime"
+
+void FPrefabsRuntimeModule::StartupModule()
+{
+}
+
+void FPrefabsRuntimeModule::ShutdownModule()
+{
+}
+
+#undef LOCTEXT_NAMESPACE
diff --git a/Engine/Plugins/Experimental/Prefabs/Source/PrefabsRuntime/Public/IPrefabsRuntimeModule.h b/Engine/Plugins/Experimental/Prefabs/Source/PrefabsRuntime/Public/IPrefabsRuntimeModule.h
new file mode 100644
--- /dev/null
+++ b/Engine/Plugins/Experimental/Prefabs/Source/PrefabsRuntime/Public/IPrefabsRuntimeModule.h
@@ -0,0 +1,36 @@
+// Copyright Epic Games, Inc. All Rights Reserved.
+#pragma once
+
+#include "CoreMinimal.h"
+#include "Modules/ModuleInterface.h"
+#include "Modules/ModuleManager.h"
+
+/**
+ * Prefabs are an invisible, linkable instantiation routine (function object or
+ * type constructor). They are used to define and later instantiate run time objects.
+ * This runtime module will implement the latter.
+ */
+class IPrefabsRuntimeModule : public IModuleInterface
+{
+public:
+	/**
+	 * Singleton-like access to this module's interface.  This is just for convenience!
+	 * Beware of calling this during the shutdown phase, though.  Your module might have been unloaded already.
+	 *
+	 * @return Returns singleton instance, loading the module on demand if needed
+	 */
+	static inline IPrefabsRuntimeModule& Get()
+	{
+		return FModuleManager::LoadModuleChecked< IPrefabsRuntimeModule >( "PrefabsRuntime" );
+	}
+
+	/**
+	 * Checks to see if this module is loaded and ready.  It is only valid to call Get() if IsAvailable() returns true.
+	 *
+	 * @return True if the module is loaded and ready to use
+	 */
+	static inline bool IsAvailable()
+	{
+		return FModuleManager::Get().IsModuleLoaded( "PrefabsRuntime" );
+	}
+};
diff --git a/Engine/Plugins/Experimental/Prefabs/Source/PrefabsUncooked/PrefabsUncooked.Build.cs b/Engine/Plugins/Experimental/Prefabs/Source/PrefabsUncooked/PrefabsUncooked.Build.cs
new file mode 100644
--- /dev/null
+++ b/Engine/Plugins/Experimental/Prefabs/Source/PrefabsUncooked/PrefabsUncooked.Build.cs
@@ -0,0 +1,19 @@
+// Copyright Epic Games, Inc. All Rights Reserved.
+
+namespace UnrealBuildTool.Rules
+{
+	public class PrefabsUncooked : ModuleRules
+	{
+		public PrefabsUncooked(ReadOnlyTargetRules Target) : base(Target)
+        {
+			PrivateDependencyModuleNames.AddRange(
+				new string[]
+				{
+					"Core",
+					"CoreUObject",
+					"Engine"
+                }
+			); 
+		}
+	}
+}
diff --git a/Engine/Plugins/Experimental/Prefabs/Source/PrefabsUncooked/Private/PrefabCompilationManager.cpp b/Engine/Plugins/Experimental/Prefabs/Source/PrefabsUncooked/Private/PrefabCompilationManager.cpp
new file mode 100644
--- /dev/null
+++ b/Engine/Plugins/Experimental/Prefabs/Source/PrefabsUncooked/Private/PrefabCompilationManager.cpp
@@ -0,0 +1,65 @@
+// Copyright Epic Games, Inc. All Rights Reserved.
+
+#include "PrefabCompilationManager.h"
+#include "Containers/Set.h"
+#include "UObject/GCObject.h"
+#include "PrefabUncooked.h"
+#include "Tickable.h"
+
+struct FPrefabCompilationManagerImpl : FGCObject, FTickableGameObject
+{
+	virtual void AddReferencedObjects(FReferenceCollector& Collector) override;
+	virtual void Tick(float DeltaTime) override;
+	virtual bool IsTickable() const override;
+	virtual bool IsTickableInEditor() const override { return true; }
+	virtual TStatId GetStatId() const override;
+	void CompilePrefab(const UPrefabUncooked* EditorPrefab);
+
+	TSet<const UPrefabUncooked*> PrefabsPendingCompilation;
+};
+
+void FPrefabCompilationManagerImpl::AddReferencedObjects(FReferenceCollector& Collector)
+{	
+	Collector.AddReferencedObjects(PrefabsPendingCompilation);
+}
+
+void FPrefabCompilationManagerImpl::Tick(float DeltaTime)
+{
+	// spread this out? Run off main thread? The latter is desirable
+	for (const UPrefabUncooked* ped : PrefabsPendingCompilation)
+	{
+		CompilePrefab(ped);
+	}
+	PrefabsPendingCompilation.Reset();
+}
+
+bool FPrefabCompilationManagerImpl::IsTickable() const
+{
+	return PrefabsPendingCompilation.Num() > 0;
+}
+
+TStatId FPrefabCompilationManagerImpl::GetStatId() const
+{
+	RETURN_QUICK_DECLARE_CYCLE_STAT(FPrefabCompilationManagerImpl, STATGROUP_Tickables);
+}
+
+void FPrefabCompilationManagerImpl::CompilePrefab(const UPrefabUncooked* EditorPrefab)
+{
+	// add a ufunction? make a lil UPrefab class to hide our factory function?
+}
+
+FPrefabCompilationManagerImpl* PFCMImpl = nullptr;
+
+void FPrefabCompilationManager::Initialize()
+{
+	if (PFCMImpl)
+	{
+		delete PFCMImpl;
+	}
+	PFCMImpl = new FPrefabCompilationManagerImpl();
+}
+
+void FPrefabCompilationManager::NotifyPrefabEdited(const UPrefabUncooked* EditorPrefab)
+{
+	PFCMImpl->PrefabsPendingCompilation.Add(EditorPrefab);
+}
diff --git a/Engine/Plugins/Experimental/Prefabs/Source/PrefabsUncooked/Private/PrefabUncooked.cpp b/Engine/Plugins/Experimental/Prefabs/Source/PrefabsUncooked/Private/PrefabUncooked.cpp
new file mode 100644
--- /dev/null
+++ b/Engine/Plugins/Experimental/Prefabs/Source/PrefabsUncooked/Private/PrefabUncooked.cpp
@@ -0,0 +1,14 @@
+// Copyright Epic Games, Inc. All Rights Reserved.
+
+#include "PrefabUncooked.h"
+#include "PrefabCompilationManager.h"
+
+void UPrefabUncooked::PostEditChangeProperty(FPropertyChangedEvent& PropertyChangedEvent)
+{
+	FPrefabCompilationManager::NotifyPrefabEdited(this);
+}
+
+void UPrefabUncooked::PostEditChangeChainProperty(FPropertyChangedChainEvent& PropertyChangedEvent)
+{
+	FPrefabCompilationManager::NotifyPrefabEdited(this);
+}
diff --git a/Engine/Plugins/Experimental/Prefabs/Source/PrefabsUncooked/Private/PrefabsUncookedModule.cpp b/Engine/Plugins/Experimental/Prefabs/Source/PrefabsUncooked/Private/PrefabsUncookedModule.cpp
new file mode 100644
--- /dev/null
+++ b/Engine/Plugins/Experimental/Prefabs/Source/PrefabsUncooked/Private/PrefabsUncookedModule.cpp
@@ -0,0 +1,30 @@
+// Copyright Epic Games, Inc. All Rights Reserved.
+
+#include "IPrefabsUncookedModule.h"
+#include "Containers/UnrealString.h"
+#include "Logging/LogMacros.h"
+#include "PrefabCompilationManager.h"
+
+DECLARE_LOG_CATEGORY_EXTERN(PrefabsRuntime, Log, All);
+
+class FPrefabsUncookedModule : public IPrefabsUncookedModule
+{
+protected:
+	virtual void StartupModule() override;
+	virtual void ShutdownModule() override;
+};
+
+IMPLEMENT_MODULE( FPrefabsUncookedModule, PrefabsUncooked);
+
+#define LOCTEXT_NAMESPACE "PrefabsUncooked"
+
+void FPrefabsUncookedModule::StartupModule()
+{
+	FPrefabCompilationManager::Initialize();
+}
+
+void FPrefabsUncookedModule::ShutdownModule()
+{
+}
+
+#undef LOCTEXT_NAMESPACE
diff --git a/Engine/Plugins/Experimental/Prefabs/Source/PrefabsUncooked/Public/IPrefabsUncookedModule.h b/Engine/Plugins/Experimental/Prefabs/Source/PrefabsUncooked/Public/IPrefabsUncookedModule.h
new file mode 100644
--- /dev/null
+++ b/Engine/Plugins/Experimental/Prefabs/Source/PrefabsUncooked/Public/IPrefabsUncookedModule.h
@@ -0,0 +1,36 @@
+// Copyright Epic Games, Inc. All Rights Reserved.
+#pragma once
+
+#include "CoreMinimal.h"
+#include "Modules/ModuleInterface.h"
+#include "Modules/ModuleManager.h"
+
+/**
+ * Prefabs are an invisible, linkable instantiation routine (function object or
+ * type constructor). They are used to define and later instantiate run time objects.
+ * This 'uncooked' module provides one implementation of the former.
+ */
+class IPrefabsUncookedModule : public IModuleInterface
+{
+public:
+	/**
+	 * Singleton-like access to this module's interface.  This is just for convenience!
+	 * Beware of calling this during the shutdown phase, though.  Your module might have been unloaded already.
+	 *
+	 * @return Returns singleton instance, loading the module on demand if needed
+	 */
+	static inline IPrefabsUncookedModule& Get()
+	{
+		return FModuleManager::LoadModuleChecked< IPrefabsUncookedModule >( "PrefabsUncooked" );
+	}
+
+	/**
+	 * Checks to see if this module is loaded and ready.  It is only valid to call Get() if IsAvailable() returns true.
+	 *
+	 * @return True if the module is loaded and ready to use
+	 */
+	static inline bool IsAvailable()
+	{
+		return FModuleManager::Get().IsModuleLoaded( "PrefabsUncooked" );
+	}
+};
diff --git a/Engine/Plugins/Experimental/Prefabs/Source/PrefabsUncooked/Public/PrefabCompilationManager.h b/Engine/Plugins/Experimental/Prefabs/Source/PrefabsUncooked/Public/PrefabCompilationManager.h
new file mode 100644
--- /dev/null
+++ b/Engine/Plugins/Experimental/Prefabs/Source/PrefabsUncooked/Public/PrefabCompilationManager.h
@@ -0,0 +1,15 @@
+// Copyright Epic Games, Inc. All Rights Reserved.
+
+#pragma once
+
+#include "Containers/Map.h"
+#include "CoreMinimal.h"
+#include "CoreTypes.h"
+
+class UPrefabUncooked;
+
+struct FPrefabCompilationManager
+{
+	static void Initialize();
+	static void NotifyPrefabEdited(const UPrefabUncooked* EditorPrefab);
+};
diff --git a/Engine/Plugins/Experimental/Prefabs/Source/PrefabsUncooked/Public/PrefabUncooked.h b/Engine/Plugins/Experimental/Prefabs/Source/PrefabsUncooked/Public/PrefabUncooked.h
new file mode 100644
--- /dev/null
+++ b/Engine/Plugins/Experimental/Prefabs/Source/PrefabsUncooked/Public/PrefabUncooked.h
@@ -0,0 +1,58 @@
+// Copyright Epic Games, Inc. All Rights Reserved.
+
+#pragma once
+
+#include "CoreMinimal.h"
+#include "UObject/ObjectMacros.h"
+#include "UObject/Object.h"
+
+#include "PrefabUncooked.generated.h"
+
+// This type is an obscure implementation detail, it is meant to be 
+// hidden from the user completely by our editor or other frontend:
+USTRUCT()
+struct FPFVariableAssignment
+{
+	GENERATED_BODY()
+	
+	// This identifier may chain/nest into an object graph:
+	UPROPERTY(EditAnywhere, Category = Prefab)
+	FString PFVariableIdentifier;
+
+	// Type could be inferred from value, but explicit will be more durable:
+	UPROPERTY(EditAnywhere, Category = Prefab)
+	FString PFVariableType;
+
+	// Stringized version of a value, again more durable than an actual value type:
+	UPROPERTY(EditAnywhere, Category = Prefab)
+	FString PFVariableValue;
+};
+
+// Structure encapsulating the state of an instance that will be created from the prefab:
+USTRUCT()
+struct FPFInstanceDeclaration
+{
+	GENERATED_BODY()
+
+	// Path to our type - the immediate goal is for this to be a native UClass or another 
+	// Prefab, but we could support other instancing facilities:
+	UPROPERTY(EditAnywhere, Category = Prefab, meta = (DisplayName = "Instance Type"))
+	FSoftObjectPath PFInstanceType;
+
+	// This will be completely hidden from the user, they will poke at a 
+	// preview object:
+	UPROPERTY(EditAnywhere, Category = Prefab)
+	TArray<FPFVariableAssignment> PFVariableAssignments;
+};
+
+UCLASS()
+class PREFABSUNCOOKED_API UPrefabUncooked : public UObject
+{
+	GENERATED_BODY()
+
+	UPROPERTY(Category = Prefab, EditAnywhere)
+	FPFInstanceDeclaration InstanceDecl; // this will be an array, sooner than later
+
+	virtual void PostEditChangeProperty(FPropertyChangedEvent& PropertyChangedEvent) override;
+	virtual void PostEditChangeChainProperty(FPropertyChangedChainEvent& PropertyChangedEvent) override;
+};
